DelegateHeaderTool updated to match recent delegate changes.

* Default pointer type added to dynamic delegate macros.
* Path to DelegateCombinations.h fixed.

diff --git a/Engine/Source/Programs/DelegateHeaderTool/Windows/Program.cs b/Engine/Source/Programs/DelegateHeaderTool/Windows/Program.cs
--- a/Engine/Source/Programs/DelegateHeaderTool/Windows/Program.cs
+++ b/Engine/Source/Programs/DelegateHeaderTool/Windows/Program.cs
@@ -134,6 +134,8 @@ namespace DelegateHeaderTool
 			var DefineDeclareDynamicDelegate = new StringBuilder( "#define DECLARE_DYNAMIC_DELEGATE" );
 			var DefineDeclareDynamicMulticastDelegate = new StringBuilder( "#define DECLARE_DYNAMIC_MULTICAST_DELEGATE" );
 
+			var DynamicDelegateDefaultPtr = "FWeakObjectPtr";
+
 
 			// Return type, always specified here even for void functions
 			var TemplateDecl = new StringBuilder( "RetValType" );
@@ -236,13 +238,13 @@ namespace DelegateHeaderTool
             DefineDeclareEvent.AppendFormat(" ) FUNC_DECLARE_EVENT( OwningType, EventName, {0}", FuncSuffix);
 			if( bSupportsRetVal )
 			{
-				DefineDeclareDynamicDelegate.AppendFormat( " ) FUNC_DECLARE_DYNAMIC_DELEGATE_RETVAL( {0}, DelegateName, DelegateName##_DelegateWrapper", FuncSuffix );
+				DefineDeclareDynamicDelegate.AppendFormat( " ) FUNC_DECLARE_DYNAMIC_DELEGATE_RETVAL( {0}, {1}, DelegateName, DelegateName##_DelegateWrapper", DynamicDelegateDefaultPtr, FuncSuffix );
 			}
 			else
 			{
-				DefineDeclareDynamicDelegate.AppendFormat( " ) FUNC_DECLARE_DYNAMIC_DELEGATE( {0}, DelegateName, DelegateName##_DelegateWrapper", FuncSuffix );
+				DefineDeclareDynamicDelegate.AppendFormat( " ) FUNC_DECLARE_DYNAMIC_DELEGATE( {0}, {1}, DelegateName, DelegateName##_DelegateWrapper", DynamicDelegateDefaultPtr, FuncSuffix );
 			}
-			DefineDeclareDynamicMulticastDelegate.AppendFormat( " ) FUNC_DECLARE_DYNAMIC_MULTICAST_DELEGATE( {0}, DelegateName, DelegateName##_DelegateWrapper", FuncSuffix );
+			DefineDeclareDynamicMulticastDelegate.AppendFormat( " ) FUNC_DECLARE_DYNAMIC_MULTICAST_DELEGATE( {0}, {1}, DelegateName, DelegateName##_DelegateWrapper", DynamicDelegateDefaultPtr, FuncSuffix );
 
 			if( bSupportsRetVal )
 			{
@@ -493,7 +495,7 @@ namespace DelegateHeaderTool
 		static void Main( string[] args )
 		{
 			// NOTE: The source file path is hard-coded relative to the directory that we expect this project to live in
-			using( var Output = new StreamWriter( "..\\..\\..\\..\\..\\..\\..\\Engine\\Source\\Runtime\\Core\\Public\\Templates\\DelegateCombinations.h" ) )
+			using( var Output = new StreamWriter( "..\\..\\..\\..\\..\\..\\..\\Engine\\Source\\Runtime\\Core\\Public\\Delegates\\DelegateCombinations.h" ) )
 			{
 				Output.WriteLine( "// Copyright 1998-2014 Epic Games, Inc. All Rights Reserved." );
 				Output.WriteLine( "" );
